Added GPU Reshape Bootstrapper for Gauntlet side module roles
- Wraps executable in headless instrumentation mode with a user workspace config

diff --git a/Engine/Source/Programs/AutomationTool/Gauntlet/Unreal/BuildSource/Gauntlet.UnrealBuildSource.cs b/Engine/Source/Programs/AutomationTool/Gauntlet/Unreal/BuildSource/Gauntlet.UnrealBuildSource.cs
--- a/Engine/Source/Programs/AutomationTool/Gauntlet/Unreal/BuildSource/Gauntlet.UnrealBuildSource.cs
+++ b/Engine/Source/Programs/AutomationTool/Gauntlet/Unreal/BuildSource/Gauntlet.UnrealBuildSource.cs
@@ -394,18 +394,22 @@ namespace Gauntlet
 			if (InUnrealPath != null)
 			{
 				// check for the editor
-				string EditorExe = Path.Combine(InUnrealPath.FullName, GetRelativeExecutablePath(UnrealTargetRole.Editor, BuildHostPlatform.Current.Platform, InConfiguration));
-
+				string EditorExe  = Path.Combine(InUnrealPath.FullName, GetRelativeExecutablePath(UnrealTargetRole.Editor, BuildHostPlatform.Current.Platform, InConfiguration));
+				string OverlayExe = Path.Combine(InUnrealPath.FullName, GetModuleOverlayExecutablePath(UnrealTargetRole.Editor, BuildHostPlatform.Current.Platform));
+				
 				if (Utils.SystemHelpers.ApplicationExists(EditorExe))
 				{
 					EditorBuild NewBuild = new EditorBuild(EditorExe, InConfiguration);
 
 					return NewBuild;
 				}
-				else
+				
+				if (Utils.SystemHelpers.ApplicationExists(OverlayExe))
 				{
-					Log.Info("No editor binaries found at {0}. Unable to create an editor build source.", EditorExe);
+					return new EditorBuild(OverlayExe, InConfiguration);
 				}
+				
+				Log.Info("No editor binaries found at {0}. Unable to create an editor build source.", EditorExe);
 			}
 			else
 			{
@@ -777,6 +781,18 @@ namespace Gauntlet
 			return CommandlineToReturn;
 		}
 
+
+		/// <summary>
+		/// Get the module overlay binary path for a particular configuration
+		/// </summary>
+		private string GetModuleOverlayExecutablePath(UnrealTargetRole role, UnrealTargetPlatform platform)
+		{
+			return Path.Join(
+				"Engine", "Binaries", platform.ToString(),
+				UnrealHelpers.CustomModuleToRoles.FirstOrDefault(module => module.Value == role).Key
+			);
+		}
+
 		/// <summary>
 		/// Given a role, platform, and config, returns the path to the binary for that config. E.g. Binaries\Win64\FooServer-Win64-Shipping.exe
 		/// </summary>
diff --git a/Engine/Source/Programs/GPUReshape/Source/GPUReshapeBootstrapper.Target.cs b/Engine/Source/Programs/GPUReshape/Source/GPUReshapeBootstrapper.Target.cs
new file mode 100644
--- /dev/null
+++ b/Engine/Source/Programs/GPUReshape/Source/GPUReshapeBootstrapper.Target.cs
@@ -0,0 +1,25 @@
+// Copyright Epic Games, Inc. All Rights Reserved.
+
+using UnrealBuildTool;
+
+[SupportedPlatforms(UnrealPlatformClass.Desktop)]
+public class GPUReshapeBootstrapperTarget : TargetRules
+{
+	public GPUReshapeBootstrapperTarget(TargetInfo Target) : base(Target)
+	{
+		Type = TargetType.Program;
+		
+		IncludeOrderVersion = EngineIncludeOrderVersion.Latest;
+		LinkType = TargetLinkType.Monolithic;
+		
+		bCompileICU = false;
+		bBuildDeveloperTools = false;		
+		bCompileAgainstEngine = false;
+		bCompileAgainstCoreUObject = false;
+		bCompileAgainstApplicationCore = false;
+		bIsBuildingConsoleApplication = true;
+
+		SolutionDirectory = "Programs/GPUReshape";
+		LaunchModuleName = "GPUReshapeBootstrapper";
+	}
+}
diff --git a/Engine/Source/Programs/GPUReshape/Source/GPUReshapeBootstrapper/GPUReshapeBootstrapper.Build.cs b/Engine/Source/Programs/GPUReshape/Source/GPUReshapeBootstrapper/GPUReshapeBootstrapper.Build.cs
new file mode 100644
--- /dev/null
+++ b/Engine/Source/Programs/GPUReshape/Source/GPUReshapeBootstrapper/GPUReshapeBootstrapper.Build.cs
@@ -0,0 +1,13 @@
+// Copyright Epic Games, Inc. All Rights Reserved.
+
+using UnrealBuildTool;
+
+public class GPUReshapeBootstrapper : ModuleRules
+{
+	public GPUReshapeBootstrapper(ReadOnlyTargetRules Target) : base(Target)
+	{
+		PublicIncludePathModuleNames.Add("Launch");
+		PrivateDependencyModuleNames.Add("Core");
+		PrivateDependencyModuleNames.Add("Projects");
+	}
+}
\ No newline at end of file
diff --git a/Engine/Source/Programs/GPUReshape/Source/GPUReshapeBootstrapper/Private/GPUReshapeBootstrapper.cpp b/Engine/Source/Programs/GPUReshape/Source/GPUReshapeBootstrapper/Private/GPUReshapeBootstrapper.cpp
new file mode 100644
--- /dev/null
+++ b/Engine/Source/Programs/GPUReshape/Source/GPUReshapeBootstrapper/Private/GPUReshapeBootstrapper.cpp
@@ -0,0 +1,119 @@
+// Copyright Epic Games, Inc. All Rights Reserved.
+
+#include "CoreMinimal.h"
+#include "RequiredProgramMainCPPInclude.h"
+
+DEFINE_LOG_CATEGORY_STATIC(LogGPUReshape, Log, All);
+
+IMPLEMENT_APPLICATION(GPUReshapeBootstrapper, "GPUReshapeBootstrapper");
+
+INT32_MAIN_INT32_ARGC_TCHAR_ARGV()
+{
+	/**
+	 * Gauntlet GPU Reshape Utility Helper
+	 *
+	 * Bootstrap Tree:
+	 *   Gauntlet -> GPUReshapeBootstrapper -> GPUReshape -> Target
+	 */
+	
+	FTaskTagScope Scope(ETaskTag::EGameThread);
+	
+	ON_SCOPE_EXIT
+	{ 
+		LLM(FLowLevelMemTracker::Get().UpdateStatsPerFrame());
+		RequestEngineExit(TEXT("Exiting"));
+		
+		FEngineLoop::AppPreExit();
+		FModuleManager::Get().UnloadModulesAtShutdown();
+		FEngineLoop::AppExit();
+	};
+
+	if (int32 Ret = GEngineLoop.PreInit(ArgC, ArgV))
+	{
+		return Ret;
+	}
+	
+#if !UE_BUILD_SHIPPING
+	if (FParse::Param(FCommandLine::Get(), TEXT("WaitForDebugger")))
+	{
+		while (!FPlatformMisc::IsDebuggerPresent())
+		{
+			FPlatformProcess::Sleep(0.1f);
+		}
+		
+		UE_DEBUG_BREAK();
+	}
+#endif // !UE_BUILD_SHIPPING
+
+	
+	FString GPUReshapePath;
+	if (!FParse::Value(FCommandLine::Get(), TEXT("GRS.Path="), GPUReshapePath))
+	{
+		UE_LOG(LogGPUReshape, Error, TEXT("GPU Reshape path not set"));
+		return 1u;
+	}
+	
+	FString AppPath;
+	if (!FParse::Value(FCommandLine::Get(), TEXT("GRS.App="), AppPath))
+	{
+		UE_LOG(LogGPUReshape, Error, TEXT("Target executable path not set"));
+		return 1u;
+	}
+	
+	FString WorkspacePath;
+	if (!FParse::Value(FCommandLine::Get(), TEXT("GRS.Workspace="), WorkspacePath))
+	{
+		UE_LOG(LogGPUReshape, Error, TEXT("Workspace path not set"));
+		return 1u;
+	}
+	
+	FString ReportPath;
+	if (!FParse::Value(FCommandLine::Get(), TEXT("GRS.Report="), ReportPath))
+	{
+		UE_LOG(LogGPUReshape, Error, TEXT("Report path not set"));
+		return 1u;
+	}
+	
+	int32 Timeout = 7200;
+	FParse::Value(FCommandLine::Get(), TEXT("GRS.Timeout="), Timeout);
+	
+	// Start in headless mode
+	FStringBuilderBase GRSCommandLine;
+	GRSCommandLine << "headless ";
+	GRSCommandLine << "-report " << "\"" << ReportPath << "\" ";
+	GRSCommandLine << "-workspace \"" << WorkspacePath << "\" ";
+	GRSCommandLine << "-timeout " << Timeout << " ";
+	GRSCommandLine << "-app " << AppPath;
+
+	for (int32_t i = 1; i < ArgC; i++)
+	{
+		GRSCommandLine << " " << ArgV[i];
+	}
+	
+	// Launch the editor bootstrapped through reshape
+	FProcHandle Handle = FPlatformProcess::CreateProc(
+		*GPUReshapePath,
+		GRSCommandLine.GetData(),
+		true,
+		false,
+		false,
+		nullptr,
+		0,
+		nullptr,
+		nullptr
+	);
+	
+	if (!Handle.IsValid())
+	{
+		UE_LOG(LogGPUReshape, Error, TEXT("Failed to launch bootstrapped application"));
+		return 1u;
+	}
+
+	// Wait for reshape to finish
+	FPlatformProcess::WaitForProc(Handle);
+
+	int32 ReturnCode = 1u;
+	FPlatformProcess::GetProcReturnCode(Handle, &ReturnCode);
+	FPlatformProcess::CloseProc(Handle);
+	return ReturnCode;
+}
