Updated Reshape Bootstrapper against new CLI
- Updated build config

diff --git a/Engine/Source/Programs/GPUReshape/Resources/BuildConfig.ini b/Engine/Source/Programs/GPUReshape/Resources/BuildConfig.ini
--- a/Engine/Source/Programs/GPUReshape/Resources/BuildConfig.ini
+++ b/Engine/Source/Programs/GPUReshape/Resources/BuildConfig.ini
@@ -1,5 +1,11 @@
 ; Build configuration overrides
 
 [Startup]
+
+; -- Compilation Environment
 ; Always enable embedded symbols, for source traceback
 r.Shaders.Symbols=1
+; Disable optimizations to get more accurate source tracebacks
+r.Shaders.Optimize=0
+; Force DXC across the board, DXBC instrumentation is supported but without symbolification
+r.D3D.ForceDXC=1
diff --git a/Engine/Source/Programs/GPUReshape/Source/GPUReshapeBootstrapper/Private/GPUReshapeBootstrapper.cpp b/Engine/Source/Programs/GPUReshape/Source/GPUReshapeBootstrapper/Private/GPUReshapeBootstrapper.cpp
--- a/Engine/Source/Programs/GPUReshape/Source/GPUReshapeBootstrapper/Private/GPUReshapeBootstrapper.cpp
+++ b/Engine/Source/Programs/GPUReshape/Source/GPUReshapeBootstrapper/Private/GPUReshapeBootstrapper.cpp
@@ -81,7 +81,7 @@ INT32_MAIN_INT32_ARGC_TCHAR_ARGV()
 	
 	// Start in headless mode
 	FStringBuilderBase GRSCommandLine;
-	GRSCommandLine << "headless ";
+	GRSCommandLine << "launch ";
 	GRSCommandLine << "-report " << "\"" << ReportPath << "\" ";
 	GRSCommandLine << "-workspace \"" << WorkspacePath << "\" ";
 	GRSCommandLine << "-timeout " << Timeout << " ";
