GPU Reshape, hooked up injection thread suspension to avoid TLS initialization order issue

diff --git a/Engine/Source/Programs/GPUReshape/Resources/Workspaces/BasicWorkspace.json b/Engine/Source/Programs/GPUReshape/Resources/Workspaces/BasicWorkspace.json
--- a/Engine/Source/Programs/GPUReshape/Resources/Workspaces/BasicWorkspace.json
+++ b/Engine/Source/Programs/GPUReshape/Resources/Workspaces/BasicWorkspace.json
@@ -2,7 +2,8 @@
 	"Config": {
 		"Coverage": true,
 		"Detail": true,
-		"RedirectOutput": true
+		"RedirectOutput": true,
+		"SuspendDeferredInitialization": true
 	},
 	"Features": [
 		"Resource Bounds",
diff --git a/Engine/Source/Programs/GPUReshape/Resources/Workspaces/ConcurrencyWorkspace.json b/Engine/Source/Programs/GPUReshape/Resources/Workspaces/ConcurrencyWorkspace.json
--- a/Engine/Source/Programs/GPUReshape/Resources/Workspaces/ConcurrencyWorkspace.json
+++ b/Engine/Source/Programs/GPUReshape/Resources/Workspaces/ConcurrencyWorkspace.json
@@ -2,7 +2,8 @@
 	"Config": {
 		"Coverage": true,
 		"Detail": true,
-		"RedirectOutput": true
+		"RedirectOutput": true,
+		"SuspendDeferredInitialization": true
 	},
 	"Features": [
 		"Concurrency"
diff --git a/Engine/Source/Programs/GPUReshape/Resources/Workspaces/InitializationWorkspace.json b/Engine/Source/Programs/GPUReshape/Resources/Workspaces/InitializationWorkspace.json
--- a/Engine/Source/Programs/GPUReshape/Resources/Workspaces/InitializationWorkspace.json
+++ b/Engine/Source/Programs/GPUReshape/Resources/Workspaces/InitializationWorkspace.json
@@ -2,7 +2,8 @@
 	"Config": {
 		"Coverage": true,
 		"Detail": true,
-		"RedirectOutput": true
+		"RedirectOutput": true,
+		"SuspendDeferredInitialization": true
 	},
 	"Features": [
 		"Initialization"
diff --git a/Engine/Source/Programs/GPUReshape/Resources/Workspaces/WaterfallWorkspace.json b/Engine/Source/Programs/GPUReshape/Resources/Workspaces/WaterfallWorkspace.json
--- a/Engine/Source/Programs/GPUReshape/Resources/Workspaces/WaterfallWorkspace.json
+++ b/Engine/Source/Programs/GPUReshape/Resources/Workspaces/WaterfallWorkspace.json
@@ -1,7 +1,8 @@
 {
 	"Config": {
 		"Detail": true,
-		"RedirectOutput": true
+		"RedirectOutput": true,
+		"SuspendDeferredInitialization": true
 	},
 	"Features": [
 		"Waterfall"
